Add a new editor subsystem to assist with tracking contextual asset/path permissions in a restricted editor mode.

diff --git a/Engine/Source/Editor/AssetPermissions/AssetPermissions.Build.cs b/Engine/Source/Editor/AssetPermissions/AssetPermissions.Build.cs
new file mode 100644
--- /dev/null
+++ b/Engine/Source/Editor/AssetPermissions/AssetPermissions.Build.cs
@@ -0,0 +1,18 @@
+// Copyright Epic Games, Inc. All Rights Reserved.
+
+using UnrealBuildTool;
+
+public class AssetPermissions : ModuleRules
+{
+	public AssetPermissions(ReadOnlyTargetRules Target) : base(Target)
+	{
+		PrivateDependencyModuleNames.AddRange(
+			new string[] {
+				"Core",
+				"CoreUObject",
+				"Engine",
+				"EditorSubsystem",
+			}
+		);
+	}
+}
\ No newline at end of file
diff --git a/Engine/Source/Editor/AssetPermissions/Private/AssetPermissions.cpp b/Engine/Source/Editor/AssetPermissions/Private/AssetPermissions.cpp
new file mode 100644
--- /dev/null
+++ b/Engine/Source/Editor/AssetPermissions/Private/AssetPermissions.cpp
@@ -0,0 +1,45 @@
+// Copyright Epic Games, Inc. All Rights Reserved.
+
+#include "AssetPermissions.h"
+
+bool UAssetPermissions::HasFiltering() const
+{
+	for (FNamePermissionLists::TConstIterator It(NamePermissionLists); It; ++It)
+	{
+		if (It->Value.HasFiltering())
+		{
+			return true;
+		}
+	}
+
+	for (FPathPermissionLists::TConstIterator It(PathPermissionLists); It; ++It)
+	{
+		if (It->Value.HasFiltering())
+		{
+			return true;
+		}
+	}
+
+	return false;
+}
+
+FNamePermissionList& UAssetPermissions::GetNamePermissionList(FName ListId)
+{
+	return NamePermissionLists.FindOrAdd(ListId);
+}
+
+FPathPermissionList& UAssetPermissions::GetPathPermissionList(FName ListId)
+{
+	return PathPermissionLists.FindOrAdd(ListId);
+}
+
+TObjectPtr<UAssetPermissions>& UAssetPermissionsFactory::GetTypedPermissionsObject(UClass* ForClass)
+{
+	TObjectPtr<UAssetPermissions>& AssetPermissionsObject = AssetPermissionsObjects.FindOrAdd(ForClass);
+	if (!AssetPermissionsObject)
+	{
+		AssetPermissionsObject = NewObject<UAssetPermissions>(this, ForClass);
+	}
+
+	return AssetPermissionsObject;
+}
\ No newline at end of file
diff --git a/Engine/Source/Editor/AssetPermissions/Private/AssetPermissionsModule.cpp b/Engine/Source/Editor/AssetPermissions/Private/AssetPermissionsModule.cpp
new file mode 100644
--- /dev/null
+++ b/Engine/Source/Editor/AssetPermissions/Private/AssetPermissionsModule.cpp
@@ -0,0 +1,23 @@
+// Copyright Epic Games, Inc. All Rights Reserved.
+
+#include "Modules/ModuleManager.h"
+
+#define LOCTEXT_NAMESPACE "FAssetPermissionsModule"
+
+class FAssetPermissionsModule : public IModuleInterface
+{
+public:
+
+	// IModuleInterface interface
+	virtual void StartupModule() override
+	{
+	}
+
+	virtual void ShutdownModule() override
+	{
+	}
+};
+
+#undef LOCTEXT_NAMESPACE
+
+IMPLEMENT_MODULE(FAssetPermissionsModule, AssetPermissions);
diff --git a/Engine/Source/Editor/AssetPermissions/Private/AssetPermissionsSubsystem.cpp b/Engine/Source/Editor/AssetPermissions/Private/AssetPermissionsSubsystem.cpp
new file mode 100644
--- /dev/null
+++ b/Engine/Source/Editor/AssetPermissions/Private/AssetPermissionsSubsystem.cpp
@@ -0,0 +1,54 @@
+// Copyright Epic Games, Inc. All Rights Reserved.
+
+#include "AssetPermissionsSubsystem.h"
+#include "AssetRegistry/AssetData.h"
+#include "AssetPermissions.h"
+
+UAssetPermissionsSubsystem::UAssetPermissionsSubsystem()
+{
+	DefaultAssetPermissionsFactory = CreateDefaultSubobject<UAssetPermissionsFactory>(TEXT("DefaultAssetPermissionsFactory"));
+}
+
+UAssetPermissionsFactory* UAssetPermissionsSubsystem::GetAssetPermissionsFactory(const UObject* ForObject)
+{
+	return GetAssetPermissionsFactory(FTopLevelAssetPath(ForObject).ToString());
+}
+
+UAssetPermissionsFactory* UAssetPermissionsSubsystem::GetAssetPermissionsFactory(const FAssetData& ForAsset)
+{
+	return GetAssetPermissionsFactory(ForAsset.GetObjectPathString());
+}
+
+UAssetPermissionsFactory* UAssetPermissionsSubsystem::GetAssetPermissionsFactory(const FStringView ForContext)
+{
+	if (ForContext.IsEmpty())
+	{
+		return DefaultAssetPermissionsFactory;
+	}
+ 
+	//	If an object path itself is not explicitly registered, then we fall back to the next highest permission level that's set.
+	//	@todo - Potentially store the map key differently to optimize for permissions set on different parts of an object path tree.
+	TObjectPtr<UAssetPermissionsFactory>* AssetPermissionsFactoryPtr = AssetPermissionsFactories.FindByHash(GetTypeHash(ForContext), ForContext);
+	if (!AssetPermissionsFactoryPtr
+		&& FPackageName::IsValidPath(ForContext))
+	{
+		FString PackageName = FTopLevelAssetPath(ForContext).GetPackageName().ToString();
+		AssetPermissionsFactoryPtr = AssetPermissionsFactories.Find(PackageName);
+		if (!AssetPermissionsFactoryPtr)
+		{
+			FString PackageRoot = FPackageName::GetPackageMountPoint(PackageName).ToString();
+			AssetPermissionsFactoryPtr = AssetPermissionsFactories.Find(PackageRoot);
+			if (!AssetPermissionsFactoryPtr)
+			{
+				return DefaultAssetPermissionsFactory;
+			}
+		}
+	}
+
+	if (AssetPermissionsFactoryPtr)
+	{
+		return *AssetPermissionsFactoryPtr;
+	}
+
+	return AssetPermissionsFactories.Add(FString(ForContext), NewObject<UAssetPermissionsFactory>(this));
+}
\ No newline at end of file
diff --git a/Engine/Source/Editor/AssetPermissions/Public/AssetPermissions.h b/Engine/Source/Editor/AssetPermissions/Public/AssetPermissions.h
new file mode 100644
--- /dev/null
+++ b/Engine/Source/Editor/AssetPermissions/Public/AssetPermissions.h
@@ -0,0 +1,85 @@
+// Copyright Epic Games, Inc. All Rights Reserved.
+
+#pragma once
+
+#include "CoreMinimal.h"
+#include "Misc/NamePermissionList.h"
+#include "Templates/SubclassOf.h"
+#include "UObject/Object.h"
+
+#include "AssetPermissions.generated.h"
+
+/**
+ * Base class for typed asset permissions.
+ * 
+ * Derive from this class to wrap one or more permission lists.
+ */
+UCLASS()
+class ASSETPERMISSIONS_API UAssetPermissions : public UObject
+{
+	GENERATED_BODY()
+
+public:
+	/**
+	 * @return TRUE if any permissions are set on this instance.
+	 */
+	bool HasFiltering() const;
+
+protected:
+	/**
+	 * Derived types can use this method to access a specific allow/deny
+	 * list for names with a specific list ID.
+	 * 
+	 * Note: The list ID is to be defined internally by the derived type.
+	 */ 
+	FNamePermissionList& GetNamePermissionList(FName ListId);
+
+	/**
+	 * Derived types can use this method to access a specific allow/deny
+	 * list for paths with a specific list ID.
+	 *
+	 * Note: The list ID is to be defined internally by the derived type.
+	 */
+	FPathPermissionList& GetPathPermissionList(FName ListId);
+
+private:
+	/** Allow/deny lists for names. Map keys are internal to derived types. */
+	typedef TMap<FName, FNamePermissionList> FNamePermissionLists;
+	FNamePermissionLists NamePermissionLists;
+
+	/** Allow/deny lists for paths. Map keys are internal to derived types. */
+	typedef TMap<FName, FPathPermissionList> FPathPermissionLists;
+	FPathPermissionLists PathPermissionLists;
+};
+
+/**
+ * Factory class for creating/registering a set of typed asset permissions.
+ */
+UCLASS()
+class ASSETPERMISSIONS_API UAssetPermissionsFactory final : public UObject
+{
+	GENERATED_BODY()
+
+public:
+	/**
+	 * Create/access asset permissions for a specific derived type. Only one
+	 * instance per type will be created/registered for this factory object.
+	 * 
+	 * @return The typed permissions object that's registered to this factory.
+	 */
+	template<typename T
+		UE_REQUIRES(std::is_base_of_v<UAssetPermissions, T>)>
+	T* GetAssetPermissions()
+	{
+		return CastChecked<T>(GetTypedPermissionsObject(T::StaticClass()));
+	}
+
+protected:
+	/** Internal helper method to register/access a specific typed instance. */
+	TObjectPtr<UAssetPermissions>& GetTypedPermissionsObject(UClass* ForClass);
+
+private:
+	/** Unique type->instance asset permissions map for this factory instance. */
+	UPROPERTY()
+	TMap<TSubclassOf<UAssetPermissions>, TObjectPtr<UAssetPermissions>> AssetPermissionsObjects;
+};
\ No newline at end of file
diff --git a/Engine/Source/Editor/AssetPermissions/Public/AssetPermissionsSubsystem.h b/Engine/Source/Editor/AssetPermissions/Public/AssetPermissionsSubsystem.h
new file mode 100644
--- /dev/null
+++ b/Engine/Source/Editor/AssetPermissions/Public/AssetPermissionsSubsystem.h
@@ -0,0 +1,74 @@
+// Copyright Epic Games, Inc. All Rights Reserved.
+
+#pragma once
+
+#include "CoreMinimal.h"
+#include "EditorSubsystem.h"
+
+#include "AssetPermissionsSubsystem.generated.h"
+
+class UAssetPermissionsFactory;
+
+/**
+ * Editor subsystem for creating/registering asset permissions. One permissions
+ * object will be created/registered per object, asset, path or custom ID string.
+ */
+UCLASS()
+class ASSETPERMISSIONS_API UAssetPermissionsSubsystem : public UEditorSubsystem
+{
+	GENERATED_BODY()
+
+public:
+	/** Default constructor. */
+	UAssetPermissionsSubsystem();
+
+	/**
+	 * Return a unique asset permissions factory for the given object context.
+	 * Note: A NULL/invalid context will return the default permissions factory.
+	 * 
+	 * @param ForObject		Input context.
+	 * @return The asset permissions factory registered for the input context.
+	 */
+	UAssetPermissionsFactory* GetAssetPermissionsFactory(const UObject* ForObject);
+
+	/**
+	 * Return a unique asset permissions factory for the given asset context.
+	 *
+	 * @param ForAsset		Input context.
+	 * @return The asset permissions factory registered for the input context.
+	 */
+	UAssetPermissionsFactory* GetAssetPermissionsFactory(const FAssetData& ForAsset);
+
+	/**
+	 * Return a unique asset permissions factory for the given string context.
+	 * 
+	 * Note: The input context can represent any valid string. However, path strings are a special case.
+	 * If the context represents a valid path, we also check the package name as well as the package root.
+	 * This allows permissions to be set on the actual asset itself, to allow for overriding package rules.
+	 * 
+	 * Example:
+	 * 
+	 *	Object path - /MyPlugin/MyContent.MyAsset
+	 *	Each of the subpaths may be registered separately:
+	 * 
+	 *	/MyPlugin						// Root level permissions
+	 *	/MyPlugin/MyContent				// Package level permissions
+	 *	/MyPlugin/MyContent.MyAsset		// Per-asset permission overrides
+	 * 
+	 * The caller is expected to handle inheriting and/or overriding of asset permission rules at each level.
+	 *
+	 * @param ForContext	Input context. Default value is the empty string.
+	 * @return The asset permissions factory registered for the input context.
+	 *		   An empty input string will return the default permissions factory.
+	 */
+	UAssetPermissionsFactory* GetAssetPermissionsFactory(const FStringView ForContext = FString());
+
+private:
+	/** Default permissions factory. Used as a fallback when no context is specified. */
+	UPROPERTY()
+	TObjectPtr<UAssetPermissionsFactory> DefaultAssetPermissionsFactory;
+
+	/** Registered asset permissions factory objects. One factory per unique context. */
+	UPROPERTY()
+	TMap<FString, TObjectPtr<UAssetPermissionsFactory>> AssetPermissionsFactories;
+};
diff --git a/Engine/Source/Editor/Kismet/Kismet.Build.cs b/Engine/Source/Editor/Kismet/Kismet.Build.cs
--- a/Engine/Source/Editor/Kismet/Kismet.Build.cs
+++ b/Engine/Source/Editor/Kismet/Kismet.Build.cs
@@ -60,7 +60,8 @@ public class Kismet : ModuleRules
 				"SubobjectEditor",
 				"SubobjectDataInterface",
 				"ToolWidgets",
-				"TraceLog"
+				"TraceLog",
+				"AssetPermissions",
 			}
 			);
 
diff --git a/Engine/Source/Editor/Kismet/Private/BlueprintAssetPermissions.cpp b/Engine/Source/Editor/Kismet/Private/BlueprintAssetPermissions.cpp
new file mode 100644
--- /dev/null
+++ b/Engine/Source/Editor/Kismet/Private/BlueprintAssetPermissions.cpp
@@ -0,0 +1,33 @@
+// Copyright Epic Games, Inc. All Rights Reserved.
+
+#include "BlueprintAssetPermissions.h"
+
+FPathPermissionList& UBlueprintAssetPermissions::GetClassPermissions()
+{
+	static const FName NAME_AllowedClasses(TEXT("AllowedClasses"));
+	return GetPathPermissionList(NAME_AllowedClasses);
+}
+
+FPathPermissionList& UBlueprintAssetPermissions::GetFunctionPermissions()
+{
+	static const FName NAME_AllowedFunctions(TEXT("AllowedFunctions"));
+	return GetPathPermissionList(NAME_AllowedFunctions);
+}
+
+FPathPermissionList& UBlueprintAssetPermissions::GetStructPermissions()
+{
+	static const FName NAME_AllowedStructs(TEXT("AllowedStructs"));
+	return GetPathPermissionList(NAME_AllowedStructs);
+}
+
+FPathPermissionList& UBlueprintAssetPermissions::GetEnumPermissions()
+{
+	static const FName NAME_AllowedEnums(TEXT("AllowedEnums"));
+	return GetPathPermissionList(NAME_AllowedEnums);
+}
+
+FNamePermissionList& UBlueprintAssetPermissions::GetPinCategoryPermissions()
+{
+	static const FName NAME_AllowedPinCategories(TEXT("AllowedPinCategories"));
+	return GetNamePermissionList(NAME_AllowedPinCategories);
+}
\ No newline at end of file
diff --git a/Engine/Source/Editor/Kismet/Public/BlueprintAssetPermissions.h b/Engine/Source/Editor/Kismet/Public/BlueprintAssetPermissions.h
new file mode 100644
--- /dev/null
+++ b/Engine/Source/Editor/Kismet/Public/BlueprintAssetPermissions.h
@@ -0,0 +1,23 @@
+// Copyright Epic Games, Inc. All Rights Reserved.
+
+#pragma once
+
+#include "AssetPermissions.h"
+
+#include "BlueprintAssetPermissions.generated.h"
+
+class UBlueprint;
+class UClass;
+
+UCLASS()
+class KISMET_API UBlueprintAssetPermissions : public UAssetPermissions
+{
+	GENERATED_BODY()
+
+public:
+	FPathPermissionList& GetClassPermissions();
+	FPathPermissionList& GetFunctionPermissions();
+	FPathPermissionList& GetStructPermissions();
+	FPathPermissionList& GetEnumPermissions();
+	FNamePermissionList& GetPinCategoryPermissions();
+};
\ No newline at end of file
