Modify permissions editor subsystem to make adding a new permissions factory an explicit operation, along with removal.

diff --git a/Engine/Source/Editor/AssetPermissions/Private/AssetPermissions.cpp b/Engine/Source/Editor/AssetPermissions/Private/AssetPermissions.cpp
--- a/Engine/Source/Editor/AssetPermissions/Private/AssetPermissions.cpp
+++ b/Engine/Source/Editor/AssetPermissions/Private/AssetPermissions.cpp
@@ -33,6 +33,11 @@ FPathPermissionList& UAssetPermissions::GetPathPermissionList(FName ListId)
 	return PathPermissionLists.FindOrAdd(ListId);
 }
 
+void UAssetPermissionsFactory::ClearAssetPermissions()
+{
+	AssetPermissionsObjects.Reset();
+}
+
 TObjectPtr<UAssetPermissions>& UAssetPermissionsFactory::GetTypedPermissionsObject(UClass* ForClass)
 {
 	TObjectPtr<UAssetPermissions>& AssetPermissionsObject = AssetPermissionsObjects.FindOrAdd(ForClass);
diff --git a/Engine/Source/Editor/AssetPermissions/Private/AssetPermissionsSubsystem.cpp b/Engine/Source/Editor/AssetPermissions/Private/AssetPermissionsSubsystem.cpp
--- a/Engine/Source/Editor/AssetPermissions/Private/AssetPermissionsSubsystem.cpp
+++ b/Engine/Source/Editor/AssetPermissions/Private/AssetPermissionsSubsystem.cpp
@@ -4,31 +4,53 @@
 #include "AssetRegistry/AssetData.h"
 #include "AssetPermissions.h"
 
-UAssetPermissionsSubsystem::UAssetPermissionsSubsystem()
+void UAssetPermissionsSubsystem::Initialize(FSubsystemCollectionBase& Collection)
 {
-	DefaultAssetPermissionsFactory = CreateDefaultSubobject<UAssetPermissionsFactory>(TEXT("DefaultAssetPermissionsFactory"));
+	DefaultAssetPermissionsFactory = NewObject<UAssetPermissionsFactory>(this, TEXT("DefaultAssetPermissionsFactory"));
 }
 
-UAssetPermissionsFactory* UAssetPermissionsSubsystem::GetAssetPermissionsFactory(const UObject* ForObject)
+void UAssetPermissionsSubsystem::Deinitialize()
 {
+	if (DefaultAssetPermissionsFactory)
+	{
+		DefaultAssetPermissionsFactory->MarkAsGarbage();
+		DefaultAssetPermissionsFactory = nullptr;
+	}
+}
+
+UAssetPermissionsFactory* UAssetPermissionsSubsystem::GetAssetPermissionsFactory(const UObject* ForObject) const
+{
+	if (!ForObject
+		|| AssetPermissionsFactories.IsEmpty())
+	{
+		return DefaultAssetPermissionsFactory;
+	}
+
 	return GetAssetPermissionsFactory(FTopLevelAssetPath(ForObject).ToString());
 }
 
-UAssetPermissionsFactory* UAssetPermissionsSubsystem::GetAssetPermissionsFactory(const FAssetData& ForAsset)
+UAssetPermissionsFactory* UAssetPermissionsSubsystem::GetAssetPermissionsFactory(const FAssetData& ForAsset) const
 {
+	if (!ForAsset.IsValid()
+		|| AssetPermissionsFactories.IsEmpty())
+	{
+		return DefaultAssetPermissionsFactory;
+	}
+
 	return GetAssetPermissionsFactory(ForAsset.GetObjectPathString());
 }
 
-UAssetPermissionsFactory* UAssetPermissionsSubsystem::GetAssetPermissionsFactory(const FStringView ForContext)
+UAssetPermissionsFactory* UAssetPermissionsSubsystem::GetAssetPermissionsFactory(const FStringView ForContext) const
 {
-	if (ForContext.IsEmpty())
+	if (ForContext.IsEmpty()
+		|| AssetPermissionsFactories.IsEmpty())
 	{
 		return DefaultAssetPermissionsFactory;
 	}
  
 	//	If an object path itself is not explicitly registered, then we fall back to the next highest permission level that's set.
 	//	@todo - Potentially store the map key differently to optimize for permissions set on different parts of an object path tree.
-	TObjectPtr<UAssetPermissionsFactory>* AssetPermissionsFactoryPtr = AssetPermissionsFactories.FindByHash(GetTypeHash(ForContext), ForContext);
+	const TObjectPtr<UAssetPermissionsFactory>* AssetPermissionsFactoryPtr = AssetPermissionsFactories.FindByHash(GetTypeHash(ForContext), ForContext);
 	if (!AssetPermissionsFactoryPtr
 		&& FPackageName::IsValidPath(ForContext))
 	{
@@ -38,10 +60,6 @@ UAssetPermissionsFactory* UAssetPermissionsSubsystem::GetAssetPermissionsFactory
 		{
 			FString PackageRoot = FPackageName::GetPackageMountPoint(PackageName).ToString();
 			AssetPermissionsFactoryPtr = AssetPermissionsFactories.Find(PackageRoot);
-			if (!AssetPermissionsFactoryPtr)
-			{
-				return DefaultAssetPermissionsFactory;
-			}
 		}
 	}
 
@@ -50,5 +68,41 @@ UAssetPermissionsFactory* UAssetPermissionsSubsystem::GetAssetPermissionsFactory
 		return *AssetPermissionsFactoryPtr;
 	}
 
-	return AssetPermissionsFactories.Add(FString(ForContext), NewObject<UAssetPermissionsFactory>(this));
+	return DefaultAssetPermissionsFactory;
+}
+
+UAssetPermissionsFactory* UAssetPermissionsSubsystem::FindOrAddAssetPermissionsFactory(const UObject* ForObject)
+{
+	if (!ForObject)
+	{
+		return DefaultAssetPermissionsFactory;
+	}
+
+	return FindOrAddAssetPermissionsFactory(FTopLevelAssetPath(ForObject).ToString());
+}
+
+UAssetPermissionsFactory* UAssetPermissionsSubsystem::FindOrAddAssetPermissionsFactory(const FAssetData& ForAsset)
+{
+	if (!ForAsset.IsValid())
+	{
+		return DefaultAssetPermissionsFactory;
+	}
+
+	return FindOrAddAssetPermissionsFactory(ForAsset.GetObjectPathString());
+}
+
+UAssetPermissionsFactory* UAssetPermissionsSubsystem::FindOrAddAssetPermissionsFactory(const FStringView ForContext)
+{
+	if (ForContext.IsEmpty())
+	{
+		return DefaultAssetPermissionsFactory;
+	}
+
+	return AssetPermissionsFactories.FindOrAdd(FString(ForContext), NewObject<UAssetPermissionsFactory>(this));
+}
+
+void UAssetPermissionsSubsystem::ClearAssetPermissions()
+{
+	AssetPermissionsFactories.Reset();
+	DefaultAssetPermissionsFactory->ClearAssetPermissions();
 }
\ No newline at end of file
diff --git a/Engine/Source/Editor/AssetPermissions/Public/AssetPermissions.h b/Engine/Source/Editor/AssetPermissions/Public/AssetPermissions.h
--- a/Engine/Source/Editor/AssetPermissions/Public/AssetPermissions.h
+++ b/Engine/Source/Editor/AssetPermissions/Public/AssetPermissions.h
@@ -74,6 +74,11 @@ public:
 		return CastChecked<T>(GetTypedPermissionsObject(T::StaticClass()));
 	}
 
+	/**
+	 * Resets the asset permissions factory and clears all typed objects.
+	 */
+	void ClearAssetPermissions();
+
 protected:
 	/** Internal helper method to register/access a specific typed instance. */
 	TObjectPtr<UAssetPermissions>& GetTypedPermissionsObject(UClass* ForClass);
diff --git a/Engine/Source/Editor/AssetPermissions/Public/AssetPermissionsSubsystem.h b/Engine/Source/Editor/AssetPermissions/Public/AssetPermissionsSubsystem.h
--- a/Engine/Source/Editor/AssetPermissions/Public/AssetPermissionsSubsystem.h
+++ b/Engine/Source/Editor/AssetPermissions/Public/AssetPermissionsSubsystem.h
@@ -19,28 +19,26 @@ class ASSETPERMISSIONS_API UAssetPermissionsSubsystem : public UEditorSubsystem
 	GENERATED_BODY()
 
 public:
-	/** Default constructor. */
-	UAssetPermissionsSubsystem();
-
 	/**
-	 * Return a unique asset permissions factory for the given object context.
-	 * Note: A NULL/invalid context will return the default permissions factory.
+	 * Return the asset permissions factory for the given object context.
 	 * 
 	 * @param ForObject		Input context.
-	 * @return The asset permissions factory registered for the input context.
+	 * @return The asset permissions factory registered for the input context, or the default
+	 *         asset permissions factory instance if no unique factory instance was registered.
 	 */
-	UAssetPermissionsFactory* GetAssetPermissionsFactory(const UObject* ForObject);
+	UAssetPermissionsFactory* GetAssetPermissionsFactory(const UObject* ForObject) const;
 
 	/**
-	 * Return a unique asset permissions factory for the given asset context.
+	 * Return the asset permissions factory for the given asset data context.
 	 *
 	 * @param ForAsset		Input context.
-	 * @return The asset permissions factory registered for the input context.
+	 * @return The asset permissions factory registered for the input context, or the default
+	 *         asset permissions factory instance if no unique factory instance was registered.
 	 */
-	UAssetPermissionsFactory* GetAssetPermissionsFactory(const FAssetData& ForAsset);
+	UAssetPermissionsFactory* GetAssetPermissionsFactory(const FAssetData& ForAsset) const;
 
 	/**
-	 * Return a unique asset permissions factory for the given string context.
+	 * Return the asset permissions factory for the given input context string.
 	 * 
 	 * Note: The input context can represent any valid string. However, path strings are a special case.
 	 * If the context represents a valid path, we also check the package name as well as the package root.
@@ -61,7 +59,42 @@ public:
 	 * @return The asset permissions factory registered for the input context.
 	 *		   An empty input string will return the default permissions factory.
 	 */
-	UAssetPermissionsFactory* GetAssetPermissionsFactory(const FStringView ForContext = FString());
+	UAssetPermissionsFactory* GetAssetPermissionsFactory(const FStringView ForContext = FString()) const;
+
+	/**
+	 * Find or add a unique asset permissions factory for the given object context.
+	 *
+	 * @param ForObject		Input context.
+	 * @return The asset permissions factory that's registered for the input context.
+	 */
+	UAssetPermissionsFactory* FindOrAddAssetPermissionsFactory(const UObject* ForObject);
+
+	/**
+	 * Find or add a unique asset permissions factory for the given asset data context.
+	 *
+	 * @param ForAsset		Input context.
+	 * @return The asset permissions factory that's registered for the input context.
+	 */
+	UAssetPermissionsFactory* FindOrAddAssetPermissionsFactory(const FAssetData& ForAsset);
+
+	/**
+	 * Find or add a unique asset permissions factory for the given string context.
+	 *
+	 * @param ForContext	Input context.
+	 * @return The asset permissions factory that's registered for the input context.
+	 */
+	UAssetPermissionsFactory* FindOrAddAssetPermissionsFactory(const FStringView ForContext);
+
+	/**
+	 * Resets the subsystem and clears all registered asset permissions factories.
+	 */
+	void ClearAssetPermissions();
+
+protected:
+	/** Begin USubsystem interface */
+	virtual void Initialize(FSubsystemCollectionBase& Collection) override;
+	virtual void Deinitialize() override;
+	/** End USubsystem interface */
 
 private:
 	/** Default permissions factory. Used as a fallback when no context is specified. */
