Add owned delegate maps to contextual asset permissions.

Notes:
- Containers are now mutable to allow public accessors to be used with a const reference.

diff --git a/Engine/Source/Editor/AssetPermissions/Private/AssetPermissions.cpp b/Engine/Source/Editor/AssetPermissions/Private/AssetPermissions.cpp
--- a/Engine/Source/Editor/AssetPermissions/Private/AssetPermissions.cpp
+++ b/Engine/Source/Editor/AssetPermissions/Private/AssetPermissions.cpp
@@ -20,15 +20,23 @@ bool UAssetPermissions::HasFiltering() const
 		}
 	}
 
+	for (FPermissionDelegates::TConstIterator It(PermissionDelegates); It; ++It)
+	{
+		if ((*It->Value).HasFiltering())
+		{
+			return true;
+		}
+	}
+
 	return false;
 }
 
-FNamePermissionList& UAssetPermissions::GetNamePermissionList(FName ListId)
+FNamePermissionList& UAssetPermissions::GetNamePermissionList(FName ListId) const
 {
 	return NamePermissionLists.FindOrAdd(ListId);
 }
 
-FPathPermissionList& UAssetPermissions::GetPathPermissionList(FName ListId)
+FPathPermissionList& UAssetPermissions::GetPathPermissionList(FName ListId) const
 {
 	return PathPermissionLists.FindOrAdd(ListId);
 }
diff --git a/Engine/Source/Editor/AssetPermissions/Public/AssetPermissions.h b/Engine/Source/Editor/AssetPermissions/Public/AssetPermissions.h
--- a/Engine/Source/Editor/AssetPermissions/Public/AssetPermissions.h
+++ b/Engine/Source/Editor/AssetPermissions/Public/AssetPermissions.h
@@ -4,6 +4,7 @@
 
 #include "CoreMinimal.h"
 #include "Misc/NamePermissionList.h"
+#include "PermissionDelegate.h"
 #include "Templates/SubclassOf.h"
 #include "UObject/Object.h"
 
@@ -32,7 +33,7 @@ protected:
 	 * 
 	 * Note: The list ID is to be defined internally by the derived type.
 	 */ 
-	FNamePermissionList& GetNamePermissionList(FName ListId);
+	FNamePermissionList& GetNamePermissionList(FName ListId) const;
 
 	/**
 	 * Derived types can use this method to access a specific allow/deny
@@ -40,16 +41,38 @@ protected:
 	 *
 	 * Note: The list ID is to be defined internally by the derived type.
 	 */
-	FPathPermissionList& GetPathPermissionList(FName ListId);
+	FPathPermissionList& GetPathPermissionList(FName ListId) const;
+
+	/**
+	 * Derived types can use this method to access a custom permission
+	 * callback with a specific delegate ID.
+	 * 
+	 * Note: Delegate ID is to be defined internally by the derived type.
+	 */
+	template<typename DelegateType>
+	TPermissionDelegate<DelegateType>& GetPermissionDelegate(FName DelegateId) const
+	{
+		TSharedRef<IPermissionDelegate>* DelegatePtr = PermissionDelegates.Find(DelegateId);
+		if (!DelegatePtr)
+		{
+			DelegatePtr = &PermissionDelegates.Add(DelegateId, MakeShared<TPermissionDelegate<DelegateType>>(DelegateId));
+		}
+		
+		return *StaticCastSharedRef<TPermissionDelegate<DelegateType>>(*DelegatePtr);
+	}
 
 private:
 	/** Allow/deny lists for names. Map keys are internal to derived types. */
 	typedef TMap<FName, FNamePermissionList> FNamePermissionLists;
-	FNamePermissionLists NamePermissionLists;
+	mutable FNamePermissionLists NamePermissionLists;
 
 	/** Allow/deny lists for paths. Map keys are internal to derived types. */
 	typedef TMap<FName, FPathPermissionList> FPathPermissionLists;
-	FPathPermissionLists PathPermissionLists;
+	mutable FPathPermissionLists PathPermissionLists;
+
+	/** Added permission delegates. Map keys are internal to derived types. */
+	typedef TMap<FName, TSharedRef<IPermissionDelegate>> FPermissionDelegates;
+	mutable FPermissionDelegates PermissionDelegates;
 };
 
 /**
diff --git a/Engine/Source/Editor/AssetPermissions/Public/PermissionDelegate.h b/Engine/Source/Editor/AssetPermissions/Public/PermissionDelegate.h
new file mode 100644
--- /dev/null
+++ b/Engine/Source/Editor/AssetPermissions/Public/PermissionDelegate.h
@@ -0,0 +1,107 @@
+// Copyright Epic Games, Inc. All Rights Reserved.
+
+#pragma once
+
+#include "CoreMinimal.h"
+
+/**
+ * Required interface for asset permission delegates.
+ */
+class IPermissionDelegate
+{
+public:
+	/** @return TRUE if the delegate has one or more registered owners. */
+	virtual bool HasFiltering() const = 0;
+};
+
+/**
+ * Template class for typed asset permission delegates.
+ * 
+ * Note: Delegate function bindings must return a true/false result.
+ */
+template<typename DelegateType
+	UE_REQUIRES(std::is_same_v<bool, typename DelegateType::RetValType>)>
+class TPermissionDelegate : public IPermissionDelegate
+{
+public:
+	explicit TPermissionDelegate(FName InDelegateId)
+		: DelegateId(InDelegateId)
+	{
+	}
+
+	virtual ~TPermissionDelegate()
+	{
+	}
+
+	// Begin IPermissionDelegate interface
+	virtual bool HasFiltering() const override
+	{
+		return !DelegateMap.IsEmpty();
+	}
+	// End IPermissionDelegate interface
+
+	/** @return The unique ID associated with the owned delegate map. */
+	FName GetId() const
+	{
+		return DelegateId;
+	}
+
+	/** @return TRUE if one or more registered delegates are bound to a function. */
+	bool IsBound() const
+	{
+		for (typename TMap<FName, DelegateType>::TConstIterator It(DelegateMap); It; ++It)
+		{
+			if (It->Value.IsBound())
+			{
+				return true;
+			}
+		}
+
+		return false;
+	}
+
+	/**
+	 * Functor to invoke the registered set of delegates based on its signature.
+	 * 
+	 * - If no delegates are registered, the functor will return TRUE (not filtered).
+	 * - If no registered delegates are bound, the functor will return FALSE (filtered).
+	 * - Otherwise, the functor will invoke all bound delegates until any one returns TRUE.
+	 */
+	template<typename ... Args>
+	bool operator()(Args&& ... args) const
+	{
+		bool bResult = DelegateMap.IsEmpty();
+		for (typename TMap<FName, DelegateType>::TConstIterator It(DelegateMap); It && !bResult; ++It)
+		{
+			if (It->Value.IsBound())
+			{
+				bResult |= It->Value.Execute(Forward<Args>(args)...);
+			}
+		}
+
+		return bResult;
+	}
+
+	/**
+	 * Register a delegate for a specific owner name.
+	 * 
+	 * @param OwnerName		Owner name to map it to.
+	 * @param InDelegate	The delegate to register.
+	 */
+	void Register(FName OwnerName, DelegateType InDelegate)
+	{
+		DelegateMap.FindOrAdd(OwnerName, InDelegate);
+	}
+
+	/**
+	 * Remove all delegates registered to an owner name.
+	 */
+	void Unregister(FName OwnerName)
+	{
+		DelegateMap.Remove(OwnerName);
+	}
+
+private:
+	FName DelegateId = NAME_None;
+	TMap<FName, DelegateType> DelegateMap;
+};
\ No newline at end of file
